Adding ability to associate a Tag with GetContentState requests and cancel requests by tag.
#tests iOS Client

diff --git a/Engine/Source/Runtime/Core/Public/GenericPlatform/GenericPlatformInstallBundleManager.h b/Engine/Source/Runtime/Core/Public/GenericPlatform/GenericPlatformInstallBundleManager.h
--- a/Engine/Source/Runtime/Core/Public/GenericPlatform/GenericPlatformInstallBundleManager.h
+++ b/Engine/Source/Runtime/Core/Public/GenericPlatform/GenericPlatformInstallBundleManager.h
@@ -296,9 +296,11 @@ public:
 	virtual FInstallBundleRequestInfo RequestUpdateContent(FName BundleName, EInstallBundleRequestFlags Flags) = 0;
 	virtual FInstallBundleRequestInfo RequestUpdateContent(TArrayView<FName> BundleNames, EInstallBundleRequestFlags Flags) = 0;
 
-	virtual void GetContentState(FName BundleName, bool bAddDependencies, FInstallBundleGetContentStateDelegate Callback) = 0;
-	virtual void GetContentState(TArrayView<FName> BundleNames, bool bAddDependencies, FInstallBundleGetContentStateDelegate Callback) = 0;
+	virtual void GetContentState(FName BundleName, bool bAddDependencies, FInstallBundleGetContentStateDelegate Callback, FName RequestTag = TEXT("None")) = 0;
+	virtual void GetContentState(TArrayView<FName> BundleNames, bool bAddDependencies, FInstallBundleGetContentStateDelegate Callback, FName RequestTag = TEXT("None")) = 0;
 
+    virtual void CancelAllGetContentStateRequestsForTag(FName RequestTag) = 0;
+    
 	virtual FInstallBundleRequestInfo RequestRemoveBundle(FName BundleName) = 0;
 
 	virtual void RequestRemoveBundleOnNextInit(FName BundleName) = 0;
diff --git a/Engine/Source/Runtime/NullInstallBundleManager/Source/NullInstallBundleManager.cpp b/Engine/Source/Runtime/NullInstallBundleManager/Source/NullInstallBundleManager.cpp
--- a/Engine/Source/Runtime/NullInstallBundleManager/Source/NullInstallBundleManager.cpp
+++ b/Engine/Source/Runtime/NullInstallBundleManager/Source/NullInstallBundleManager.cpp
@@ -45,20 +45,25 @@ class FNullInstallBundleManager : public IPlatformInstallBundleManager
 		return RetInfo;
 	}
 
-	virtual void GetContentState(FName BundleName, bool bAddDependencies, FInstallBundleGetContentStateDelegate Callback) override
+	virtual void GetContentState(FName BundleName, bool bAddDependencies, FInstallBundleGetContentStateDelegate Callback, FName RequestTag) override
 	{
 		FInstallBundleContentState State;
 		State.State = EInstallBundleContentState::UpToDate;
 		Callback.ExecuteIfBound(State);
 	}
 
-	virtual void GetContentState(TArrayView<FName> BundleNames, bool bAddDependencies, FInstallBundleGetContentStateDelegate Callback) override
+	virtual void GetContentState(TArrayView<FName> BundleNames, bool bAddDependencies, FInstallBundleGetContentStateDelegate Callback, FName RequestTag) override
 	{
 		FInstallBundleContentState State;
 		State.State = EInstallBundleContentState::UpToDate;
 		Callback.ExecuteIfBound(State);
 	}
 
+    virtual void CancelAllGetContentStateRequestsForTag(FName RequestTag) override
+    {
+        
+    }
+    
 	virtual void RequestRemoveBundleOnNextInit(FName BundleName) override
 	{
 
