Update coopyright notices

diff --git a/Engine/Source/Runtime/Core/Private/GenericPlatform/GenericPlatformInstallBundleManager.cpp b/Engine/Source/Runtime/Core/Private/GenericPlatform/GenericPlatformInstallBundleManager.cpp
--- a/Engine/Source/Runtime/Core/Private/GenericPlatform/GenericPlatformInstallBundleManager.cpp
+++ b/Engine/Source/Runtime/Core/Private/GenericPlatform/GenericPlatformInstallBundleManager.cpp
@@ -1,4 +1,4 @@
-// Copyright 1998-2018 Epic Games, Inc. All Rights Reserved.
+// Copyright 1998-2019 Epic Games, Inc. All Rights Reserved.
 
 #include "GenericPlatform/GenericPlatformInstallBundleManager.h"
 
diff --git a/Engine/Source/Runtime/Core/Private/Misc/EmbeddedCommunication.cpp b/Engine/Source/Runtime/Core/Private/Misc/EmbeddedCommunication.cpp
--- a/Engine/Source/Runtime/Core/Private/Misc/EmbeddedCommunication.cpp
+++ b/Engine/Source/Runtime/Core/Private/Misc/EmbeddedCommunication.cpp
@@ -1,4 +1,4 @@
-// Copyright 1998-2018 Epic Games, Inc. All Rights Reserved.
+// Copyright 1998-2019 Epic Games, Inc. All Rights Reserved.
 
 // Core includes.
 #include "Misc/EmbeddedCommunication.h"
diff --git a/Engine/Source/Runtime/Core/Public/GenericPlatform/GenericPlatformInstallBundleManager.h b/Engine/Source/Runtime/Core/Public/GenericPlatform/GenericPlatformInstallBundleManager.h
--- a/Engine/Source/Runtime/Core/Public/GenericPlatform/GenericPlatformInstallBundleManager.h
+++ b/Engine/Source/Runtime/Core/Public/GenericPlatform/GenericPlatformInstallBundleManager.h
@@ -1,4 +1,4 @@
-// Copyright 1998-2018 Epic Games, Inc. All Rights Reserved.
+// Copyright 1998-2019 Epic Games, Inc. All Rights Reserved.
 
 #pragma once
 
diff --git a/Engine/Source/Runtime/Core/Public/Misc/EmbeddedCommunication.h b/Engine/Source/Runtime/Core/Public/Misc/EmbeddedCommunication.h
--- a/Engine/Source/Runtime/Core/Public/Misc/EmbeddedCommunication.h
+++ b/Engine/Source/Runtime/Core/Public/Misc/EmbeddedCommunication.h
@@ -1,4 +1,4 @@
-// Copyright 1998-2018 Epic Games, Inc. All Rights Reserved.
+// Copyright 1998-2019 Epic Games, Inc. All Rights Reserved.
 
 #pragma once
 
diff --git a/Engine/Source/Runtime/NullInstallBundleManager/NullInstallBundleManager.build.cs b/Engine/Source/Runtime/NullInstallBundleManager/NullInstallBundleManager.build.cs
--- a/Engine/Source/Runtime/NullInstallBundleManager/NullInstallBundleManager.build.cs
+++ b/Engine/Source/Runtime/NullInstallBundleManager/NullInstallBundleManager.build.cs
@@ -1,4 +1,4 @@
-// Copyright 1998-2018 Epic Games, Inc. All Rights Reserved.
+// Copyright 1998-2019 Epic Games, Inc. All Rights Reserved.
 
 using UnrealBuildTool;
 
diff --git a/Engine/Source/Runtime/NullInstallBundleManager/Source/NullInstallBundleManager.cpp b/Engine/Source/Runtime/NullInstallBundleManager/Source/NullInstallBundleManager.cpp
--- a/Engine/Source/Runtime/NullInstallBundleManager/Source/NullInstallBundleManager.cpp
+++ b/Engine/Source/Runtime/NullInstallBundleManager/Source/NullInstallBundleManager.cpp
@@ -1,4 +1,4 @@
-// Copyright 1998-2018 Epic Games, Inc. All Rights Reserved.
+// Copyright 1998-2019 Epic Games, Inc. All Rights Reserved.
 
 #include "GenericPlatform/GenericPlatformInstallBundleManager.h"
 #include "Modules/ModuleManager.h"
