Fix module name

diff --git a/Engine/Source/Runtime/NullInstallBundleManager/Source/NullInstallBundleManager.cpp b/Engine/Source/Runtime/NullInstallBundleManager/Source/NullInstallBundleManager.cpp
--- a/Engine/Source/Runtime/NullInstallBundleManager/Source/NullInstallBundleManager.cpp
+++ b/Engine/Source/Runtime/NullInstallBundleManager/Source/NullInstallBundleManager.cpp
@@ -86,4 +86,4 @@ private:
 	TUniquePtr<FNullInstallBundleManager> InstallBundleManager;
 };
 
-IMPLEMENT_MODULE(FNullInstallBundleManagerModule, NullInstallBundleManagerModule);
+IMPLEMENT_MODULE(FNullInstallBundleManagerModule, NullInstallBundleManager);
