Install bundle manager refactor
- Pause / Resume support
- Cancel support
- Make TArrayView usage const correct

diff --git a/Engine/Source/Runtime/InstallBundleManager/Private/BundlePrereqCombinedStatusHelper.cpp b/Engine/Source/Runtime/InstallBundleManager/Private/BundlePrereqCombinedStatusHelper.cpp
--- a/Engine/Source/Runtime/InstallBundleManager/Private/BundlePrereqCombinedStatusHelper.cpp
+++ b/Engine/Source/Runtime/InstallBundleManager/Private/BundlePrereqCombinedStatusHelper.cpp
@@ -87,12 +87,14 @@ void FBundlePrereqCombinedStatusHelper::SetupDelegates()
 	CleanUpDelegates();
 	
 	IInstallBundleManager::InstallBundleCompleteDelegate.AddRaw(this, &FBundlePrereqCombinedStatusHelper::OnBundleInstallComplete);
+	IInstallBundleManager::PausedBundleDelegate.AddRaw(this, &FBundlePrereqCombinedStatusHelper::OnBundleInstallPauseChanged);
 	TickHandle = FTicker::GetCoreTicker().AddTicker(FTickerDelegate::CreateRaw(this, &FBundlePrereqCombinedStatusHelper::Tick));
 }
 
 void FBundlePrereqCombinedStatusHelper::CleanUpDelegates()
 {
 	IInstallBundleManager::InstallBundleCompleteDelegate.RemoveAll(this);
+	IInstallBundleManager::PausedBundleDelegate.RemoveAll(this);
 	if (TickHandle.IsValid())
 	{
 		FTicker::GetCoreTicker().RemoveTicker(TickHandle);
@@ -164,8 +166,7 @@ void FBundlePrereqCombinedStatusHelper::UpdateBundleCache()
 			//Copy progress to the cache as long as we have progress to copy.
 			if (BundleProgress.IsSet())
 			{
-				FInstallBundleStatus& CachedStatus = BundleStatusCache.FindOrAdd(BundleName);
-				CachedStatus = BundleProgress.GetValue();
+				BundleStatusCache.Add(BundleName, BundleProgress.GetValue());
 			}
 		}
 	}
@@ -365,4 +366,16 @@ void FBundlePrereqCombinedStatusHelper::OnBundleInstallComplete(FInstallBundleRe
 	}
 }
 
+// It's not really necessary to have this, but it allows for a fallback if GetBundleProgress() returns null.
+// Normally that shouldn't happen, but right now its handy while I refactor bundle progress.
+void FBundlePrereqCombinedStatusHelper::OnBundleInstallPauseChanged(FInstallBundlePauseInfo PauseInfo)
+{
+	const bool bWasRequiredBundle = RequiredBundleNames.Contains(PauseInfo.BundleName);
+	if (bWasRequiredBundle)
+	{
+		FInstallBundleStatus& BundleCacheInfo = BundleStatusCache.FindOrAdd(PauseInfo.BundleName);
+		BundleCacheInfo.PauseFlags = PauseInfo.PauseFlags;
+	}
+}
+
 //PRAGMA_ENABLE_OPTIMIZATION
diff --git a/Engine/Source/Runtime/InstallBundleManager/Public/BundlePrereqCombinedStatusHelper.h b/Engine/Source/Runtime/InstallBundleManager/Public/BundlePrereqCombinedStatusHelper.h
--- a/Engine/Source/Runtime/InstallBundleManager/Public/BundlePrereqCombinedStatusHelper.h
+++ b/Engine/Source/Runtime/InstallBundleManager/Public/BundlePrereqCombinedStatusHelper.h
@@ -62,6 +62,7 @@ private:
 	
 	//Called so we can track when a bundle is finished
 	void OnBundleInstallComplete(FInstallBundleRequestResultInfo CompletedBundleInfo);
+	void OnBundleInstallPauseChanged(FInstallBundlePauseInfo PauseInfo);
 	
 	float GetCombinedProgressPercent() const;
 	float GetIndividualWeightedProgressPercent(const FInstallBundleStatus& Bundle) const;
diff --git a/Engine/Source/Runtime/InstallBundleManager/Public/InstallBundleManagerInterface.h b/Engine/Source/Runtime/InstallBundleManager/Public/InstallBundleManagerInterface.h
--- a/Engine/Source/Runtime/InstallBundleManager/Public/InstallBundleManagerInterface.h
+++ b/Engine/Source/Runtime/InstallBundleManager/Public/InstallBundleManagerInterface.h
@@ -133,32 +133,35 @@ public:
 	virtual EInstallBundleManagerInitState GetInitState() const = 0;
 
 	virtual FInstallBundleRequestInfo RequestUpdateContent(FName BundleName, EInstallBundleRequestFlags Flags) = 0;
-	virtual FInstallBundleRequestInfo RequestUpdateContent(TArrayView<FName> BundleNames, EInstallBundleRequestFlags Flags) = 0;
+	virtual FInstallBundleRequestInfo RequestUpdateContent(TArrayView<const FName> BundleNames, EInstallBundleRequestFlags Flags) = 0;
 
 	virtual void GetContentState(FName BundleName, EInstallBundleGetContentStateFlags Flags, bool bAddDependencies, FInstallBundleGetContentStateDelegate Callback, FName RequestTag = TEXT("None")) = 0;
-	virtual void GetContentState(TArrayView<FName> BundleNames, EInstallBundleGetContentStateFlags Flags, bool bAddDependencies, FInstallBundleGetContentStateDelegate Callback, FName RequestTag = TEXT("None")) = 0;
+	virtual void GetContentState(TArrayView<const FName> BundleNames, EInstallBundleGetContentStateFlags Flags, bool bAddDependencies, FInstallBundleGetContentStateDelegate Callback, FName RequestTag = TEXT("None")) = 0;
 
     virtual void CancelAllGetContentStateRequestsForTag(FName RequestTag) = 0;
     
-	virtual void RequestRemoveContentOnNextInit(FName RemoveName, TArrayView<FName> KeepNames = TArrayView<FName>()) = 0;
-	virtual void RequestRemoveContentOnNextInit(TArrayView<FName> RemoveNames, TArrayView<FName> KeepNames = TArrayView<FName>()) = 0;
+	virtual void RequestRemoveContentOnNextInit(FName RemoveName, TArrayView<const FName> KeepNames = TArrayView<const FName>()) = 0;
+	virtual void RequestRemoveContentOnNextInit(TArrayView<const FName> RemoveNames, TArrayView<const FName> KeepNames = TArrayView<const FName>()) = 0;
 
 	virtual void CancelRequestRemoveContentOnNextInit(FName BundleName) = 0;
-	virtual void CancelRequestRemoveContentOnNextInit(TArrayView<FName> BundleName) = 0;
+	virtual void CancelRequestRemoveContentOnNextInit(TArrayView<const FName> BundleNames) = 0;
 
-	virtual void CancelBundle(FName BundleName, EInstallBundleCancelFlags Flags) = 0;
+	virtual void CancelUpdateContent(FName BundleName, EInstallBundleCancelFlags Flags) = 0;
+	virtual void CancelUpdateContent(TArrayView<const FName> BundleNames, EInstallBundleCancelFlags Flags) = 0;
 
-	virtual void CancelAllBundles(EInstallBundleCancelFlags Flags) = 0;
+	virtual void PauseUpdateContent(FName BundleName) = 0;
+	virtual void PauseUpdateContent(TArrayView<const FName> BundleNames) = 0;
 
-	virtual bool PauseBundle(FName BundleName) = 0;
+	virtual void ResumeUpdateContent(FName BundleName) = 0;
+	virtual void ResumeUpdateContent(TArrayView<const FName> BundleNames) = 0;
 
-	virtual void ResumeBundle(FName BundleName) = 0;
-
-	virtual void RequestPausedBundleCallback() const = 0;
+	virtual void RequestPausedBundleCallback() = 0;
 
 	virtual TOptional<FInstallBundleStatus> GetBundleProgress(FName BundleName) const = 0;
 
+	virtual EInstallBundleRequestFlags GetModifyableContentRequestFlags() const = 0;
 	virtual void UpdateContentRequestFlags(FName BundleName, EInstallBundleRequestFlags AddFlags, EInstallBundleRequestFlags RemoveFlags) = 0;
+	virtual void UpdateContentRequestFlags(TArrayView<const FName> BundleNames, EInstallBundleRequestFlags AddFlags, EInstallBundleRequestFlags RemoveFlags) = 0;
 
 	virtual bool IsNullInterface() const = 0;
 
diff --git a/Engine/Source/Runtime/InstallBundleManager/Public/InstallBundleSourceInterface.h b/Engine/Source/Runtime/InstallBundleManager/Public/InstallBundleSourceInterface.h
--- a/Engine/Source/Runtime/InstallBundleManager/Public/InstallBundleSourceInterface.h
+++ b/Engine/Source/Runtime/InstallBundleManager/Public/InstallBundleSourceInterface.h
@@ -24,15 +24,6 @@ public:
 	// Returns the how this source should be weighted when combined with other sources
 	virtual float GetSourceWeight() const { return 1.0f; }
 
-	// Returns true if this source can remove content on init
-	virtual bool SupportsRemoveContentOnNextInit() const { return false; }
-
-	// Returns true if this source can cancel requests
-	virtual bool SupportsCanceling() const { return false; }
-
-	// Returns true if this source can pause requests
-	virtual bool SupportsPausing() const { return false; }
-
 	// Called once by bundle manager after constructing the bundle source
 	virtual void Init(
 		TSharedRef<FConfigFile> InstallBundleConfig, 
@@ -50,7 +41,7 @@ public:
 
 	// Gets the state of content on disk
 	// BundleNames contains all dependencies and has been deduped
-	virtual void GetContentState(TArrayView<FName> BundleNames, EInstallBundleGetContentStateFlags Flags, FInstallBundleGetContentStateDelegate Callback) = 0;
+	virtual void GetContentState(TArrayView<const FName> BundleNames, EInstallBundleGetContentStateFlags Flags, FInstallBundleGetContentStateDelegate Callback) = 0;
 
 	struct FRequestUpdateContentBundleContext
 	{
@@ -67,28 +58,22 @@ public:
 
 	// Returns true if content is scheduled to be removed the next time the source is initialized
 	// BundleNames contains all dependencies and has been deduped
-	virtual bool RequestRemoveContentOnNextInit(TArrayView<FName> RemoveNames) { return false; }
+	virtual bool RequestRemoveContentOnNextInit(TArrayView<const FName> RemoveNames) { return false; }
 	// Call to cancel the removal of any content scheduled for removal the next time the source is initialized
 	// Returns true if all bundles were canceled
-	virtual bool CancelRequestRemoveContentOnNextInit(TArrayView<FName> BundleNames) { return false;  }
+	virtual bool CancelRequestRemoveContentOnNextInit(TArrayView<const FName> BundleNames) { return false; }
 
 	// Cancel the install for the specified bundles
-	// BundleNames contains all dependencies and has been deduped
-	virtual void CancelBundles(TArrayView<FName> BundleNames, EInstallBundleCancelFlags Flags) {}
-
-	// Pause/Resume bundles.  Returns true if any bundles were paused/resumed
-	virtual bool PauseBundles(TArrayView<FName> BundleNames) { return false; }
-	virtual bool ResumeBundles(TArrayView<FName> BundleNames) { return false; }
+	virtual void CancelBundles(TArrayView<const FName> BundleNames, EInstallBundleCancelFlags Flags) {}
 
-	// Called by bundle manager to update pause status.  Returns a struct indicating current pause status and if that status has changed.
-	virtual FInstallBundleSourcePauseInfo UpdateBundlePauseInfo(FName BundleName) 
-	{ 
-		FInstallBundleSourcePauseInfo Info;
-		Info.BundleName = BundleName;
-		return Info;
-	}
+	// User Pause/Resume bundles.
+	virtual void UserPauseBundles(TArrayView<const FName> BundleNames) {}
+	virtual void UserResumeBundles(TArrayView<const FName> BundleNames) {}
 
-	// TODO: UpdateContentRequestFlags - Probably need to rework the whole thing.  The only flag it makes sense for is CheckForCellularDataUsage
+	// UpdateContentRequestFlags - Allow some flags to be updated for in flight requests
+	// Currently only CheckForCellularDataUsage is supported
+	virtual EInstallBundleRequestFlags GetModifyableContentRequestFlags() const { return EInstallBundleRequestFlags::None; }
+	virtual void UpdateContentRequestFlags(TArrayView<const FName> BundleNames, EInstallBundleRequestFlags AddFlags, EInstallBundleRequestFlags RemoveFlags) {}
 
 	// TODO: GetBundleProgress() - Not sure how this will work yet
 
diff --git a/Engine/Source/Runtime/NullInstallBundleManager/Source/NullInstallBundleManager.cpp b/Engine/Source/Runtime/NullInstallBundleManager/Source/NullInstallBundleManager.cpp
--- a/Engine/Source/Runtime/NullInstallBundleManager/Source/NullInstallBundleManager.cpp
+++ b/Engine/Source/Runtime/NullInstallBundleManager/Source/NullInstallBundleManager.cpp
@@ -38,7 +38,7 @@ class FNullInstallBundleManager : public IInstallBundleManager
 		FInstallBundleRequestInfo RetInfo;
 		return RetInfo;
 	}
-	virtual FInstallBundleRequestInfo RequestUpdateContent(TArrayView<FName> BundleNames, EInstallBundleRequestFlags Flags) override
+	virtual FInstallBundleRequestInfo RequestUpdateContent(TArrayView<const FName> BundleNames, EInstallBundleRequestFlags Flags) override
 	{
 		FInstallBundleRequestInfo RetInfo;
 		return RetInfo;
@@ -51,7 +51,7 @@ class FNullInstallBundleManager : public IInstallBundleManager
 		Callback.ExecuteIfBound(State);
 	}
 
-	virtual void GetContentState(TArrayView<FName> BundleNames, EInstallBundleGetContentStateFlags Flags, bool bAddDependencies, FInstallBundleGetContentStateDelegate Callback, FName RequestTag) override
+	virtual void GetContentState(TArrayView<const FName> BundleNames, EInstallBundleGetContentStateFlags Flags, bool bAddDependencies, FInstallBundleGetContentStateDelegate Callback, FName RequestTag) override
 	{
 		FInstallBundleCombinedContentState State;
 		State.State = EInstallBundleContentState::UpToDate;
@@ -62,10 +62,10 @@ class FNullInstallBundleManager : public IInstallBundleManager
 	{
 	}
 
-	virtual void RequestRemoveContentOnNextInit(FName RemoveName, TArrayView<FName> KeepNames = TArrayView<FName>()) override
+	virtual void RequestRemoveContentOnNextInit(FName RemoveName, TArrayView<const FName> KeepNames = TArrayView<const FName>()) override
 	{
 	}
-	virtual void RequestRemoveContentOnNextInit(TArrayView<FName> RemoveNames, TArrayView<FName> KeepNames = TArrayView<FName>()) override
+	virtual void RequestRemoveContentOnNextInit(TArrayView<const FName> RemoveNames, TArrayView<const FName> KeepNames = TArrayView<const FName>()) override
 	{
 	}
 
@@ -73,31 +73,37 @@ class FNullInstallBundleManager : public IInstallBundleManager
 	{
 
 	}
-	virtual void CancelRequestRemoveContentOnNextInit(TArrayView<FName> BundleNames) override
+	virtual void CancelRequestRemoveContentOnNextInit(TArrayView<const FName> BundleNames) override
 	{
 	}
 
-	virtual void CancelBundle(FName BundleName, EInstallBundleCancelFlags Flags) override
+	virtual void CancelUpdateContent(FName BundleName, EInstallBundleCancelFlags Flags) override
 	{
 
+	}
+	virtual void CancelUpdateContent(TArrayView<const FName> BundleNames, EInstallBundleCancelFlags Flags) override
+	{
 	}
 
-	virtual void CancelAllBundles(EInstallBundleCancelFlags Flags) override
+	virtual void PauseUpdateContent(FName BundleName) override
 	{
 
 	}
-
-	virtual bool PauseBundle(FName BundleName) override
+	virtual void PauseUpdateContent(TArrayView<const FName> BundleNames) override
 	{
-		return false;
+
 	}
 
-	virtual void ResumeBundle(FName BundleName) override
+	virtual void ResumeUpdateContent(FName BundleName) override
 	{
 
 	}
+	virtual void ResumeUpdateContent(TArrayView<const FName> BundleNames) override
+	{
 
-	virtual void RequestPausedBundleCallback() const override
+	}
+
+	virtual void RequestPausedBundleCallback() override
 	{
 
 	}
@@ -107,9 +113,17 @@ class FNullInstallBundleManager : public IInstallBundleManager
 		return TOptional<FInstallBundleStatus>();
 	}
 
+	virtual EInstallBundleRequestFlags GetModifyableContentRequestFlags() const override
+	{
+		return EInstallBundleRequestFlags::None;
+	}
 	virtual void UpdateContentRequestFlags(FName BundleName, EInstallBundleRequestFlags AddFlags, EInstallBundleRequestFlags RemoveFlags) override
 	{
 
+	}
+	virtual void UpdateContentRequestFlags(TArrayView<const FName> BundleNames, EInstallBundleRequestFlags AddFlags, EInstallBundleRequestFlags RemoveFlags) override
+	{
+
 	}
 
 	virtual bool IsNullInterface() const override
